[IOS] Remove the whitelist from the IOSReplayKit plugin so that it will be able to be referenced on all platforms (including the edtior).
- In this case it is not IOS, it will default to stub functions.
- Is not enabled by default.

diff --git a/Engine/Plugins/Runtime/IOSReplayKit/IOSReplayKit.uplugin b/Engine/Plugins/Runtime/IOSReplayKit/IOSReplayKit.uplugin
--- a/Engine/Plugins/Runtime/IOSReplayKit/IOSReplayKit.uplugin
+++ b/Engine/Plugins/Runtime/IOSReplayKit/IOSReplayKit.uplugin
@@ -14,11 +14,7 @@
 		{
 			"Name": "IOSReplayKit",
 			"Type": "Runtime",
-			"LoadingPhase": "PreDefault",
-			"WhitelistPlatforms": 
-			[
-				"IOS"
-			]
+			"LoadingPhase": "PreDefault"
 		}
 	]
 }
