Add missing header guard
#rnx

diff --git a/Engine/Plugins/Runtime/IOSReplayKit/Source/IOSReplayKit/Private/ReplayKitRecorder.h b/Engine/Plugins/Runtime/IOSReplayKit/Source/IOSReplayKit/Private/ReplayKitRecorder.h
--- a/Engine/Plugins/Runtime/IOSReplayKit/Source/IOSReplayKit/Private/ReplayKitRecorder.h
+++ b/Engine/Plugins/Runtime/IOSReplayKit/Source/IOSReplayKit/Private/ReplayKitRecorder.h
@@ -1,5 +1,5 @@
 // Copyright Epic Games, Inc. All Rights Reserved.
-
+#pragma once
 
 #include "CoreTypes.h"
 
