Fix TVOS build issues with ReplayKit

diff --git a/Engine/Plugins/Runtime/IOSReplayKit/Source/IOSReplayKit/Private/ReplayKitRecorder.cpp b/Engine/Plugins/Runtime/IOSReplayKit/Source/IOSReplayKit/Private/ReplayKitRecorder.cpp
--- a/Engine/Plugins/Runtime/IOSReplayKit/Source/IOSReplayKit/Private/ReplayKitRecorder.cpp
+++ b/Engine/Plugins/Runtime/IOSReplayKit/Source/IOSReplayKit/Private/ReplayKitRecorder.cpp
@@ -26,9 +26,10 @@ NSString* _Nullable _captureFilePath;
 - (void) initializeWithMicrophoneEnabled:(BOOL)bMicrophoneEnabled withCameraEnabled:(BOOL)bCameraEnabled {
 	_screenRecorder = [RPScreenRecorder sharedRecorder];
 	[_screenRecorder setDelegate:self];
-	
+#if !PLATFORM_TVOS
 	[_screenRecorder setMicrophoneEnabled:bMicrophoneEnabled];
 	[_screenRecorder setCameraEnabled:bCameraEnabled];
+#endif
 }
 
 - (void)startRecording {
@@ -203,12 +204,14 @@ NSString* _Nullable _captureFilePath;
                 {
                     NSLog(@"finishWritingWithCompletionHandler");
                     
+#if !PLATFORM_TVOS
                     if (UIVideoAtPathIsCompatibleWithSavedPhotosAlbum(_captureFilePath))
                     {
                         UISaveVideoAtPathToSavedPhotosAlbum(_captureFilePath, nil, nil, nil);
                         NSLog(@"capture saved to album");
                     }
-                    
+#endif
+
                     [_captureFilePath release];
                     _captureFilePath = nullptr;
                     
